Editor: moved SFilePathPicker to DesktopWidgets module, because AppFramework cannot use Developer functionality, such as IDesktopPlatform

diff --git a/Engine/Source/Developer/DesktopWidgets/DesktopWidgets.Build.cs b/Engine/Source/Developer/DesktopWidgets/DesktopWidgets.Build.cs
new file mode 100644
--- /dev/null
+++ b/Engine/Source/Developer/DesktopWidgets/DesktopWidgets.Build.cs
@@ -0,0 +1,32 @@
+// Copyright 1998-2014 Epic Games, Inc. All Rights Reserved.
+
+using UnrealBuildTool;
+
+public class DesktopWidgets : ModuleRules
+{
+	public DesktopWidgets(TargetInfo Target)
+	{
+		PublicDependencyModuleNames.AddRange(
+			new string[] {
+				"Core",
+				"Slate",
+				"SlateCore",
+			}
+		);
+
+		PrivateDependencyModuleNames.AddRange(
+			new string[] {
+                "DesktopPlatform",
+                "InputCore",
+			}
+		);
+
+		PrivateIncludePaths.AddRange(
+			new string[] {
+				"Developer/DesktopWidgets/Private",
+				"Developer/DesktopWidgets/Private/Widgets",
+                "Developer/DesktopWidgets/Private/Widgets/Input",
+			}
+		);
+	}
+}
diff --git a/Engine/Source/Developer/DesktopWidgets/Private/DesktopWidgetsModule.cpp b/Engine/Source/Developer/DesktopWidgets/Private/DesktopWidgetsModule.cpp
new file mode 100644
--- /dev/null
+++ b/Engine/Source/Developer/DesktopWidgets/Private/DesktopWidgetsModule.cpp
@@ -0,0 +1,29 @@
+// Copyright 1998-2014 Epic Games, Inc. All Rights Reserved.
+
+#include "DesktopWidgetsPrivatePCH.h"
+
+
+static const FName DesktopWidgetsTabName("DesktopWidgets");
+
+
+/**
+ * Implements the DesktopWidgets module.
+ */
+class FDesktopWidgetsModule
+	: public IModuleInterface
+{
+public:
+
+	// IModuleInterface interface
+	
+	virtual void StartupModule( ) override
+	{
+	}
+
+	virtual void ShutdownModule( ) override
+	{
+	}
+};
+
+
+IMPLEMENT_MODULE(FDesktopWidgetsModule, DesktopWidgets);
diff --git a/Engine/Source/Developer/DesktopWidgets/Private/DesktopWidgetsPrivatePCH.h b/Engine/Source/Developer/DesktopWidgets/Private/DesktopWidgetsPrivatePCH.h
new file mode 100644
--- /dev/null
+++ b/Engine/Source/Developer/DesktopWidgets/Private/DesktopWidgetsPrivatePCH.h
@@ -0,0 +1,15 @@
+// Copyright 1998-2014 Epic Games, Inc. All Rights Reserved.
+
+#pragma once
+
+#include "DesktopWidgets.h"
+
+
+/* Private dependencies
+ *****************************************************************************/
+
+#include "DesktopPlatformModule.h"
+
+
+/* Private includes
+ *****************************************************************************/
diff --git a/Engine/Source/Developer/DesktopWidgets/Private/Widgets/Input/SFilePathPicker.cpp b/Engine/Source/Developer/DesktopWidgets/Private/Widgets/Input/SFilePathPicker.cpp
new file mode 100644
--- /dev/null
+++ b/Engine/Source/Developer/DesktopWidgets/Private/Widgets/Input/SFilePathPicker.cpp
@@ -0,0 +1,108 @@
+// Copyright 1998-2014 Epic Games, Inc. All Rights Reserved.
+
+#include "DesktopWidgetsPrivatePCH.h"
+#include "SFilePathPicker.h"
+
+
+#define LOCTEXT_NAMESPACE "SFilePathPicker"
+
+
+/* SFilePathPicker interface
+ *****************************************************************************/
+
+void SFilePathPicker::Construct( const FArguments& InArgs )
+{
+	BrowseDirectory = InArgs._BrowseDirectory;
+	BrowseTitle = InArgs._BrowseTitle;
+	FilePath = InArgs._FilePath;
+	FileTypeFilter = InArgs._FileTypeFilter;
+	OnPathPicked = InArgs._OnPathPicked;
+
+	ChildSlot
+	[
+		SNew(SHorizontalBox)
+
+		+ SHorizontalBox::Slot()
+			.FillWidth(1.0f)
+			.VAlign(VAlign_Center)
+			[
+				// file path text box
+				SAssignNew(TextBox, SEditableTextBox)
+					.Text(this, &SFilePathPicker::HandleTextBoxText)
+					.Font(InArgs._Font)
+					.SelectAllTextWhenFocused(true)
+					.ClearKeyboardFocusOnCommit(false)
+					.OnTextCommitted(this, &SFilePathPicker::HandleTextBoxTextCommitted)
+					.SelectAllTextOnCommit(false)
+					.IsReadOnly(InArgs._IsReadOnly)
+			]
+
+		+ SHorizontalBox::Slot()
+			.AutoWidth()
+			.Padding(4.0f, 0.0f, 0.0f, 0.0f)
+			.VAlign(VAlign_Center)
+			[
+				// browse button
+				SNew(SButton)
+					.ButtonStyle(InArgs._BrowseButtonStyle)
+					.ToolTipText(InArgs._BrowseButtonToolTip)
+					.OnClicked(this, &SFilePathPicker::HandleBrowseButtonClicked)
+					.ContentPadding(2.0f)
+					.ForegroundColor(FSlateColor::UseForeground())
+					.IsFocusable(false)
+					[
+						SNew(SImage)
+							.Image(InArgs._BrowseButtonImage)
+							.ColorAndOpacity(FSlateColor::UseForeground())
+					]
+			]
+	];
+}
+
+
+/* SFilePathPicker callbacks
+ *****************************************************************************/
+
+FReply SFilePathPicker::HandleBrowseButtonClicked( )
+{
+	IDesktopPlatform* DesktopPlatform = FDesktopPlatformModule::Get();
+
+	if (DesktopPlatform == nullptr)
+	{
+		return FReply::Handled();
+	}
+
+	const FString DefaultPath = BrowseDirectory.IsSet()
+		? BrowseDirectory.Get()
+		: FPaths::GetPath(FilePath.Get());
+
+	// show the file browse dialog
+	TSharedPtr<SWindow> ParentWindow = FSlateApplication::Get().FindWidgetWindow(AsShared());
+	void* ParentWindowHandle = (ParentWindow.IsValid() && ParentWindow->GetNativeWindow().IsValid())
+		? ParentWindow->GetNativeWindow()->GetOSWindowHandle()
+		: nullptr;
+
+	TArray<FString> OutFiles;
+
+	if (DesktopPlatform->OpenFileDialog(ParentWindowHandle, BrowseTitle.Get().ToString(), DefaultPath, TEXT(""), FileTypeFilter.Get(), EFileDialogFlags::None, OutFiles))
+	{
+		OnPathPicked.ExecuteIfBound(OutFiles[0]);
+	}
+
+	return FReply::Handled();
+}
+
+
+FText SFilePathPicker::HandleTextBoxText( ) const
+{
+	return FText::FromString(FilePath.Get());
+}
+
+
+void SFilePathPicker::HandleTextBoxTextCommitted( const FText& NewText, ETextCommit::Type CommitInfo )
+{
+	OnPathPicked.ExecuteIfBound(NewText.ToString());
+}
+
+
+#undef LOCTEXT_NAMESPACE
diff --git a/Engine/Source/Developer/DesktopWidgets/Public/DesktopWidgets.h b/Engine/Source/Developer/DesktopWidgets/Public/DesktopWidgets.h
new file mode 100644
--- /dev/null
+++ b/Engine/Source/Developer/DesktopWidgets/Public/DesktopWidgets.h
@@ -0,0 +1,15 @@
+// Copyright 1998-2014 Epic Games, Inc. All Rights Reserved.
+
+#pragma once
+
+
+/* Public dependencies
+ *****************************************************************************/
+
+#include "Core.h"
+#include "ModuleInterface.h"
+#include "Slate.h"
+
+
+/* Public includes
+ *****************************************************************************/
diff --git a/Engine/Source/Developer/DesktopWidgets/Public/Widgets/Input/SFilePathPicker.h b/Engine/Source/Developer/DesktopWidgets/Public/Widgets/Input/SFilePathPicker.h
new file mode 100644
--- /dev/null
+++ b/Engine/Source/Developer/DesktopWidgets/Public/Widgets/Input/SFilePathPicker.h
@@ -0,0 +1,100 @@
+// Copyright 1998-2014 Epic Games, Inc. All Rights Reserved.
+
+#pragma once
+
+
+/**
+ * Declares a delegate that is executed when a file was picked in the SFilePathPicker widget.
+ *
+ * The first parameter will contain the path to the picked file.
+ */
+DECLARE_DELEGATE_OneParam(FOnPathPicked, const FString& /*PickedPath*/);
+
+
+/**
+ * Implements an editable text box with a browse button.
+ */
+class DESKTOPWIDGETS_API SFilePathPicker
+	: public SCompoundWidget
+{
+public:
+
+	SLATE_BEGIN_ARGS(SFilePathPicker)
+		: _BrowseButtonToolTip(NSLOCTEXT("SFilePathPicker", "BrowseButtonToolTip", "Choose a file from this computer"))
+		, _FileTypeFilter(TEXT("All files (*.*)|*.*"))
+		, _Font()
+		, _IsReadOnly(false)
+	{ }
+
+		/** Browse button image resource. */
+		SLATE_ATTRIBUTE(const FSlateBrush*, BrowseButtonImage)
+
+		/** Browse button visual style. */
+		SLATE_STYLE_ARGUMENT(FButtonStyle, BrowseButtonStyle)
+
+		/** Browse button tool tip text. */
+		SLATE_ATTRIBUTE(FText, BrowseButtonToolTip)
+
+		/** The directory to browse by default */
+		SLATE_ATTRIBUTE(FString, BrowseDirectory)
+
+		/** Title for the browse dialog window. */
+		SLATE_ATTRIBUTE(FText, BrowseTitle)
+
+		/** The currently selected file path. */
+		SLATE_ATTRIBUTE(FString, FilePath)
+
+		/** File type filter string. */
+		SLATE_ATTRIBUTE(FString, FileTypeFilter)
+
+		/** Font color and opacity of the path text box. */
+		SLATE_ATTRIBUTE(FSlateFontInfo, Font)
+
+		/** Whether the path text box can be modified by the user. */
+		SLATE_ATTRIBUTE(bool, IsReadOnly)
+
+		/** Called when a file path has been picked. */
+		SLATE_EVENT(FOnPathPicked, OnPathPicked)
+
+	SLATE_END_ARGS()
+
+	/**
+	 * Constructs a new widget.
+	 *
+	 * @param InArgs The construction arguments.
+	 */
+	void Construct( const FArguments& InArgs );
+
+private:
+
+	// Callback for clicking the browse button.
+	FReply HandleBrowseButtonClicked( );
+
+	// Callback for getting the text in the path text box.
+	FText HandleTextBoxText( ) const;
+
+	// Callback for committing the text in the path text box.
+	void HandleTextBoxTextCommitted( const FText& NewText, ETextCommit::Type /*CommitInfo*/ );
+
+private:
+
+	/** Holds the directory path to browse by default. */
+	TAttribute<FString> BrowseDirectory;
+
+	/** Holds the title for the browse dialog window. */
+	TAttribute<FText> BrowseTitle;
+
+	/** Holds the currently selected file path. */
+	TAttribute<FString> FilePath;
+
+	/** Holds the file type filter string. */
+	TAttribute<FString> FileTypeFilter;
+
+	/** Holds the editable text box. */
+	TSharedPtr<SEditableTextBox> TextBox;
+
+private:
+
+	/** Holds a delegate that is executed when a file was picked. */
+	FOnPathPicked OnPathPicked;
+};
diff --git a/Engine/Source/Editor/DetailCustomizations/DetailCustomizations.Build.cs b/Engine/Source/Editor/DetailCustomizations/DetailCustomizations.Build.cs
--- a/Engine/Source/Editor/DetailCustomizations/DetailCustomizations.Build.cs
+++ b/Engine/Source/Editor/DetailCustomizations/DetailCustomizations.Build.cs
@@ -13,6 +13,7 @@ public class DetailCustomizations : ModuleRules
                 "AppFramework",
 				"Core",
 				"CoreUObject",
+				"DesktopWidgets",
 				"Engine",
                 "InputCore",
 				"Slate",
diff --git a/Engine/Source/Runtime/AppFramework/AppFramework.Build.cs b/Engine/Source/Runtime/AppFramework/AppFramework.Build.cs
--- a/Engine/Source/Runtime/AppFramework/AppFramework.Build.cs
+++ b/Engine/Source/Runtime/AppFramework/AppFramework.Build.cs
@@ -16,7 +16,6 @@ public class AppFramework : ModuleRules
 
 		PrivateDependencyModuleNames.AddRange(
 			new string[] {
-                "DesktopPlatform",
                 "InputCore",
 			}
 		);
@@ -24,10 +23,13 @@ public class AppFramework : ModuleRules
 		PrivateIncludePaths.AddRange(
 			new string[] {
 				"Runtime/AppFramework/Private",
+                "Runtime/AppFramework/Private/Framework",
                 "Runtime/AppFramework/Private/Framework/Testing",
 				"Runtime/AppFramework/Private/Widgets",
                 "Runtime/AppFramework/Private/Widgets/Colors",
+                "Runtime/AppFramework/Private/Widgets/Input",
                 "Runtime/AppFramework/Private/Widgets/Testing",
+                "Runtime/AppFramework/Private/Widgets/Workflow",
 			}
 		);
 	}
diff --git a/Engine/Source/Runtime/AppFramework/Private/AppFrameworkPrivatePCH.h b/Engine/Source/Runtime/AppFramework/Private/AppFrameworkPrivatePCH.h
--- a/Engine/Source/Runtime/AppFramework/Private/AppFrameworkPrivatePCH.h
+++ b/Engine/Source/Runtime/AppFramework/Private/AppFrameworkPrivatePCH.h
@@ -8,8 +8,6 @@
 /* Private dependencies
  *****************************************************************************/
 
-#include "DesktopPlatformModule.h"
-
 
 /* Private includes
  *****************************************************************************/
diff --git a/Engine/Source/Runtime/AppFramework/Private/Widgets/Input/SFilePathPicker.cpp b/Engine/Source/Runtime/AppFramework/Private/Widgets/Input/SFilePathPicker.cpp
deleted file mode 100644
diff --git a/Engine/Source/Runtime/AppFramework/Public/Widgets/Input/SFilePathPicker.h b/Engine/Source/Runtime/AppFramework/Public/Widgets/Input/SFilePathPicker.h
deleted file mode 100644
