DesktopWidgets: Made public headers compilable individually

diff --git a/Engine/Source/Developer/DesktopWidgets/Private/DesktopWidgetsModule.cpp b/Engine/Source/Developer/DesktopWidgets/Private/DesktopWidgetsModule.cpp
--- a/Engine/Source/Developer/DesktopWidgets/Private/DesktopWidgetsModule.cpp
+++ b/Engine/Source/Developer/DesktopWidgets/Private/DesktopWidgetsModule.cpp
@@ -1,6 +1,7 @@
 // Copyright 1998-2014 Epic Games, Inc. All Rights Reserved.
 
 #include "DesktopWidgetsPrivatePCH.h"
+#include "ModuleInterface.h"
 
 
 static const FName DesktopWidgetsTabName("DesktopWidgets");
@@ -16,13 +17,8 @@ public:
 
 	// IModuleInterface interface
 	
-	virtual void StartupModule( ) override
-	{
-	}
-
-	virtual void ShutdownModule( ) override
-	{
-	}
+	virtual void StartupModule() override { }
+	virtual void ShutdownModule() override { }
 };
 
 
diff --git a/Engine/Source/Developer/DesktopWidgets/Private/DesktopWidgetsPrivatePCH.h b/Engine/Source/Developer/DesktopWidgets/Private/DesktopWidgetsPrivatePCH.h
--- a/Engine/Source/Developer/DesktopWidgets/Private/DesktopWidgetsPrivatePCH.h
+++ b/Engine/Source/Developer/DesktopWidgets/Private/DesktopWidgetsPrivatePCH.h
@@ -2,12 +2,11 @@
 
 #pragma once
 
-#include "DesktopWidgets.h"
-
 
 /* Private dependencies
  *****************************************************************************/
 
+#include "Core.h"
 #include "DesktopPlatformModule.h"
 
 
diff --git a/Engine/Source/Developer/DesktopWidgets/Public/DesktopWidgets.h b/Engine/Source/Developer/DesktopWidgets/Public/DesktopWidgets.h
deleted file mode 100644
diff --git a/Engine/Source/Developer/DesktopWidgets/Public/Widgets/Input/SFilePathPicker.h b/Engine/Source/Developer/DesktopWidgets/Public/Widgets/Input/SFilePathPicker.h
--- a/Engine/Source/Developer/DesktopWidgets/Public/Widgets/Input/SFilePathPicker.h
+++ b/Engine/Source/Developer/DesktopWidgets/Public/Widgets/Input/SFilePathPicker.h
@@ -2,6 +2,8 @@
 
 #pragma once
 
+#include "SlateBasics.h"
+
 
 /**
  * Declares a delegate that is executed when a file was picked in the SFilePathPicker widget.
@@ -67,13 +69,13 @@ public:
 
 private:
 
-	// Callback for clicking the browse button.
+	/** Callback for clicking the browse button. */
 	FReply HandleBrowseButtonClicked( );
 
-	// Callback for getting the text in the path text box.
+	/** Callback for getting the text in the path text box. */
 	FText HandleTextBoxText( ) const;
 
-	// Callback for committing the text in the path text box.
+	/** Callback for committing the text in the path text box. */
 	void HandleTextBoxTextCommitted( const FText& NewText, ETextCommit::Type /*CommitInfo*/ );
 
 private:
